CrashReport - Compilation fixes for SymbolDebugger

diff --git a/Engine/Source/Developer/CrashDebugHelper/Public/CrashDebugHelper.h b/Engine/Source/Developer/CrashDebugHelper/Public/CrashDebugHelper.h
--- a/Engine/Source/Developer/CrashDebugHelper/Public/CrashDebugHelper.h
+++ b/Engine/Source/Developer/CrashDebugHelper/Public/CrashDebugHelper.h
@@ -468,6 +468,18 @@ public:
 	 */
 	virtual bool Init();
 
+	/** Get the depot name */
+	const FString& GetDepotName() const
+	{
+		return DepotName;
+	}
+
+	/** Set the depot name */
+	void SetDepotName( const FString& InDepotName )
+	{
+		DepotName = InDepotName;
+	}
+
 	/**
 	 *	Parse the given crash dump, determining EngineVersion of the build that produced it - if possible. 
 	 *
diff --git a/Engine/Source/Programs/SymbolDebugger/Private/SymbolDebugger.cpp b/Engine/Source/Programs/SymbolDebugger/Private/SymbolDebugger.cpp
--- a/Engine/Source/Programs/SymbolDebugger/Private/SymbolDebugger.cpp
+++ b/Engine/Source/Programs/SymbolDebugger/Private/SymbolDebugger.cpp
@@ -42,28 +42,15 @@ void FSymbolDebugger_AsyncInspect::DoWork()
 	}
 	else if ((EngineVersionName.Len() > 0) || (ChangelistName.Len() > 0))
 	{
-		bool bEngineVersion = (EngineVersionName.Len() > 0);
-
 		FCrashDebugHelperModule& CrashHelperModule = FModuleManager::LoadModuleChecked<FCrashDebugHelperModule>(FName("CrashDebugHelper"));
 		ICrashDebugHelper* Handler = CrashHelperModule.Get();
 		check(Handler != NULL);
 
-		FString FoundLabel;
-		if (bEngineVersion == true)
-		{
-			FoundLabel = Handler->GetLabelFromEngineVersion(FCString::Atoi(*EngineVersionName));
-		}
-		else
-		{
-			FoundLabel = Handler->GetLabelFromChangelistNumber(FCString::Atoi(*ChangelistName));
-		}
+		const FString FoundLabel = Handler->GetLabelFromChangelistNumber(FCString::Atoi(*ChangelistName));
 
 		if (FoundLabel.Len() > 0)
 		{
-			if (bEngineVersion == true)
-			{
-				Result_EngineVersionName = EngineVersionName;
-			}
+			Result_EngineVersionName = TEXT( "" );
 			Result_LabelName = FoundLabel;
 		}
 		else
