Programs: NULL to nullptr

diff --git a/Engine/Source/Programs/ShaderCompileWorker/Private/ShaderCompileWorker.cpp b/Engine/Source/Programs/ShaderCompileWorker/Private/ShaderCompileWorker.cpp
--- a/Engine/Source/Programs/ShaderCompileWorker/Private/ShaderCompileWorker.cpp
+++ b/Engine/Source/Programs/ShaderCompileWorker/Private/ShaderCompileWorker.cpp
@@ -39,7 +39,7 @@ const TArray<const IShaderFormat*>& GetShaderFormats()
 		for (int32 Index = 0; Index < Modules.Num(); Index++)
 		{
 			IShaderFormat* Format = FModuleManager::LoadModuleChecked<IShaderFormatModule>(Modules[Index]).GetShaderFormat();
-			if (Format != NULL)
+			if (Format != nullptr)
 			{
 				Results.Add(Format);
 			}
@@ -65,7 +65,7 @@ const IShaderFormat* FindShaderFormat(FName Name)
 		}
 	}
 
-	return NULL;
+	return nullptr;
 }
 	
 /** Processes a compilation job. */
@@ -217,7 +217,7 @@ private:
 	/** Opens an input file, trying multiple times if necessary. */
 	FArchive* OpenInputFile()
 	{
-		FArchive* InputFile = NULL;
+		FArchive* InputFile = nullptr;
 		bool bFirstOpenTry = true;
 		while(!InputFile && !GIsRequestingExit)
 		{
@@ -305,7 +305,7 @@ private:
 
 	FArchive* CreateOutputArchive()
 	{
-		FArchive* OutputFilePtr = NULL;
+		FArchive* OutputFilePtr = nullptr;
 		if (CommunicationMode == ThroughFile)
 		{
 			const double StartTime = FPlatformTime::Seconds();
@@ -440,7 +440,7 @@ private:
 				bool bParentStillRunning = true;
 				HANDLE ParentProcessHandle = OpenProcess(SYNCHRONIZE, false, ParentProcessId);
 				// If we couldn't open the process then it is no longer running, exit
-				if (ParentProcessHandle == NULL)
+				if (ParentProcessHandle == nullptr)
 				{
 					if (!IsUsingNamedPipes())
 					{
diff --git a/Engine/Source/Programs/SlateViewer/Private/Mac/SlateViewerMainMac.cpp b/Engine/Source/Programs/SlateViewer/Private/Mac/SlateViewerMainMac.cpp
--- a/Engine/Source/Programs/SlateViewer/Private/Mac/SlateViewerMainMac.cpp
+++ b/Engine/Source/Programs/SlateViewer/Private/Mac/SlateViewerMainMac.cpp
@@ -26,7 +26,7 @@ static FString GSavedCommandLine;
     [appleEventManager setEventHandler:self andSelector:@selector(handleQuitEvent:withReplyEvent:) forEventClass:kCoreEventClass andEventID:kAEQuitApplication];
 	
 	FPlatformMisc::SetGracefulTerminationHandler();
-	FPlatformMisc::SetCrashHandler(NULL);
+	FPlatformMisc::SetCrashHandler(nullptr);
 	
 	RunSlateViewer(*GSavedCommandLine);
 
diff --git a/Engine/Source/Programs/SymbolDebugger/Private/SymbolDebugger.cpp b/Engine/Source/Programs/SymbolDebugger/Private/SymbolDebugger.cpp
--- a/Engine/Source/Programs/SymbolDebugger/Private/SymbolDebugger.cpp
+++ b/Engine/Source/Programs/SymbolDebugger/Private/SymbolDebugger.cpp
@@ -22,7 +22,7 @@ void FSymbolDebugger_AsyncInspect::DoWork()
 {
 	FCrashDebugHelperModule& CrashHelperModule = FModuleManager::LoadModuleChecked<FCrashDebugHelperModule>(FName("CrashDebugHelper"));
 	ICrashDebugHelper* Handler = CrashHelperModule.Get();
-	check(Handler != NULL);
+	check(Handler != nullptr);
 
 	if (CrashDumpName.Len() > 0)
 	{
@@ -44,7 +44,7 @@ void FSymbolDebugger_AsyncInspect::DoWork()
 	{
 		FCrashDebugHelperModule& CrashHelperModule = FModuleManager::LoadModuleChecked<FCrashDebugHelperModule>(FName("CrashDebugHelper"));
 		ICrashDebugHelper* Handler = CrashHelperModule.Get();
-		check(Handler != NULL);
+		check(Handler != nullptr);
 
 		const FString FoundLabel = Handler->GetLabelFromChangelistNumber(FCString::Atoi(*ChangelistName));
 
@@ -76,7 +76,7 @@ void FSymbolDebugger_AsyncSyncFiles::DoWork()
 {
 	FCrashDebugHelperModule& CrashHelperModule = FModuleManager::LoadModuleChecked<FCrashDebugHelperModule>(FName("CrashDebugHelper"));
 	ICrashDebugHelper* Handler = CrashHelperModule.Get();
-	check(Handler != NULL);
+	check(Handler != nullptr);
 
 	if (Handler->SyncRequiredFilesForDebuggingFromLabel(SourceControlLabel, Platform) == true)
 	{
@@ -124,7 +124,7 @@ void FSymbolDebugger_ProcessCrashDump::DoWork()
 
 	FCrashDebugHelperModule& CrashHelperModule = FModuleManager::LoadModuleChecked<FCrashDebugHelperModule>(FName("CrashDebugHelper"));
 	ICrashDebugHelper* Handler = CrashHelperModule.Get();
-	check(Handler != NULL);
+	check(Handler != nullptr);
 
 	FCrashDebugInfo CrashDebugInfo;
 	if (Handler->ParseCrashDump(CrashDumpName, CrashDebugInfo) == true)
@@ -170,7 +170,7 @@ FSymbolDebugger::FSymbolDebugger()
 
 	FCrashDebugHelperModule& CrashHelperModule = FModuleManager::LoadModuleChecked<FCrashDebugHelperModule>(FName("CrashDebugHelper"));
 	ICrashDebugHelper* Handler = CrashHelperModule.Get();
-	check(Handler != NULL);
+	check(Handler != nullptr);
 	DepotName = Handler->GetDepotName();
 }
 
@@ -225,7 +225,7 @@ bool FSymbolDebugger::OnFileOpen(TSharedRef<SWidget> ParentWidget)
 	TArray<FString> OpenFilenames;
 	IDesktopPlatform* DesktopPlatform = FDesktopPlatformModule::Get();
 	bool bOpened = false;
-	if (DesktopPlatform != NULL)
+	if (DesktopPlatform != nullptr)
 	{
 		TSharedPtr<SWindow> ParentWindow = FSlateApplication::Get().FindWidgetWindow(ParentWidget);
 		void* ParentWindowHandle = (ParentWindow.IsValid() && ParentWindow->GetNativeWindow().IsValid()) ? ParentWindow->GetNativeWindow()->GetOSWindowHandle() : nullptr;
@@ -311,7 +311,7 @@ bool FSymbolDebugger::SetTextField(SSymbolDebugger::ESymbolDebuggerTextFields In
 			DepotName = InNewName;
 			FCrashDebugHelperModule& CrashHelperModule = FModuleManager::LoadModuleChecked<FCrashDebugHelperModule>(FName("CrashDebugHelper"));
 			ICrashDebugHelper* Handler = CrashHelperModule.Get();
-			check(Handler != NULL);
+			check(Handler != nullptr);
 			Handler->SetDepotName(DepotName);
 		}
 		break;
@@ -537,7 +537,7 @@ void FSymbolDebugger::Tick()
 		// Initialize source control and show the login window
 		FCrashDebugHelperModule& CrashHelperModule = FModuleManager::LoadModuleChecked<FCrashDebugHelperModule>(FName("CrashDebugHelper"));
 		ICrashDebugHelper* Handler = CrashHelperModule.Get();
-		check(Handler != NULL);
+		check(Handler != nullptr);
 		Handler->InitSourceControl(true);
 
 		s_bFirstTick = false;
@@ -554,7 +554,7 @@ void FSymbolDebugger::Tick()
 				EngineVersionName = InspectionTask->GetTask().GetResults_EngineVersion();
 			}
 
-			InspectionTask = NULL;
+			InspectionTask = nullptr;
 
 			LastAction = CurrentAction;
 			CurrentAction = SSymbolDebugger::DebugAction_None;
@@ -566,7 +566,7 @@ void FSymbolDebugger::Tick()
 		{
 			bSyncSucceeded = SyncTask->GetTask().DidSucceed();
 
-			SyncTask = NULL;
+			SyncTask = nullptr;
 
 			LastAction = CurrentAction;
 			CurrentAction = SSymbolDebugger::DebugAction_None;
@@ -578,7 +578,7 @@ void FSymbolDebugger::Tick()
 		{
 			bLaunchDebugSucceeded = LaunchDebuggerTask->GetTask().DidSucceed();
 
-			LaunchDebuggerTask = NULL;
+			LaunchDebuggerTask = nullptr;
 
 			LastAction = CurrentAction;
 			CurrentAction = SSymbolDebugger::DebugAction_None;
@@ -590,7 +590,7 @@ void FSymbolDebugger::Tick()
 		{
 			bProcessCrashDumpSucceeded = ProcessCrashDumpTask->GetTask().DidSucceed();
 
-			ProcessCrashDumpTask = NULL;
+			ProcessCrashDumpTask = nullptr;
 
 			LastAction = CurrentAction;
 			CurrentAction = SSymbolDebugger::DebugAction_None;
diff --git a/Engine/Source/Programs/UnrealLightmass/Private/Lighting/SampleVolume.cpp b/Engine/Source/Programs/UnrealLightmass/Private/Lighting/SampleVolume.cpp
--- a/Engine/Source/Programs/UnrealLightmass/Private/Lighting/SampleVolume.cpp
+++ b/Engine/Source/Programs/UnrealLightmass/Private/Lighting/SampleVolume.cpp
@@ -215,7 +215,7 @@ void FStaticLightingSystem::CalculateVolumeSamples()
 
 			//@todo - can this be presized more accurately?
 			VolumeLightingSamples.Empty(FMath::Max<int32>(5000, LandscapeEstimateNum));
-			FStaticLightingMappingContext MappingContext(NULL, *this);
+			FStaticLightingMappingContext MappingContext(nullptr, *this);
 			// Octree used to keep track of where existing samples have been placed
 			FVolumeLightingProximityOctree VolumeLightingOctree(VolumeBounds.Origin, VolumeBounds.BoxExtent.GetMax());
 			// Octree used for interpolating lighting for debugging
@@ -335,8 +335,8 @@ void FStaticLightingSystem::CalculateVolumeSamples()
 								const FLightRay Ray(
 									SamplePosition,
 									SamplePosition - FVector4(0,0,VolumeBounds.BoxExtent.Z * 2.0f),
-									NULL,
-									NULL
+									nullptr,
+									nullptr
 									);
 								FLightRayIntersection Intersection;
 								// Trace a ray straight down to find which level's geometry we are over, 
@@ -501,8 +501,8 @@ FGatheredLightSample FStaticLightingSystem::InterpolatePrecomputedVolumeIncident
 			const FLightRay SampleRay
 				(Vertex.WorldPosition + Vertex.WorldTangentZ * SceneConstants.VisibilityNormalOffsetSampleRadiusScale * SampleRadius, 
 				VolumeSample.GetPosition(), 
-				NULL, 
-				NULL);
+				nullptr, 
+				nullptr);
 			AggregateMesh.IntersectLightRay(SampleRay, false, false, false, RayCache, Intersection);
 			if (!Intersection.bIntersects)
 			*/
