Fix up symbol debugger to not try and load modules outside main thread

diff --git a/Engine/Source/Programs/SymbolDebugger/Private/SymbolDebugger.cpp b/Engine/Source/Programs/SymbolDebugger/Private/SymbolDebugger.cpp
--- a/Engine/Source/Programs/SymbolDebugger/Private/SymbolDebugger.cpp
+++ b/Engine/Source/Programs/SymbolDebugger/Private/SymbolDebugger.cpp
@@ -10,8 +10,9 @@
 //-----------------------------------------------------------------------------
 //	FSymbolDebugger_AsyncInspect
 //-----------------------------------------------------------------------------
-FSymbolDebugger_AsyncInspect::FSymbolDebugger_AsyncInspect(const FString& InCrashDumpName, const FString& InEngineVersion, const FString& InChangelist)
-	: AskedToAbortCount( 0 )
+FSymbolDebugger_AsyncInspect::FSymbolDebugger_AsyncInspect(const FString& InCrashDumpName, const FString& InEngineVersion, const FString& InChangelist, ICrashDebugHelper* InCrashHelperModule)
+	: CrashHelperModule(InCrashHelperModule)
+	, AskedToAbortCount( 0 )
 	, CrashDumpName(InCrashDumpName)
 	, EngineVersionName(InEngineVersion)
 	, ChangelistName(InChangelist)
@@ -20,14 +21,10 @@ FSymbolDebugger_AsyncInspect::FSymbolDebugger_AsyncInspect(const FString& InCras
 
 void FSymbolDebugger_AsyncInspect::DoWork()
 {
-	FCrashDebugHelperModule& CrashHelperModule = FModuleManager::LoadModuleChecked<FCrashDebugHelperModule>(FName("CrashDebugHelper"));
-	ICrashDebugHelper* Handler = CrashHelperModule.Get();
-	check(Handler != nullptr);
-
 	if (CrashDumpName.Len() > 0)
 	{
 		FCrashDebugInfo CrashDebugInfo;
-		if (Handler->ParseCrashDump(CrashDumpName, CrashDebugInfo) == true)
+		if (CrashHelperModule->ParseCrashDump(CrashDumpName, CrashDebugInfo) == true)
 		{
 			Result_LabelName = CrashDebugInfo.SourceControlLabel;
 			Result_EngineVersionName = FString::FromInt(CrashDebugInfo.EngineVersion);
@@ -42,11 +39,7 @@ void FSymbolDebugger_AsyncInspect::DoWork()
 	}
 	else if ((EngineVersionName.Len() > 0) || (ChangelistName.Len() > 0))
 	{
-		FCrashDebugHelperModule& CrashHelperModule = FModuleManager::LoadModuleChecked<FCrashDebugHelperModule>(FName("CrashDebugHelper"));
-		ICrashDebugHelper* Handler = CrashHelperModule.Get();
-		check(Handler != nullptr);
-
-		const FString FoundLabel = Handler->GetLabelFromChangelistNumber(FCString::Atoi(*ChangelistName));
+		const FString FoundLabel = CrashHelperModule->GetLabelFromChangelistNumber(FCString::Atoi(*ChangelistName));
 
 		if (FoundLabel.Len() > 0)
 		{
@@ -64,8 +57,9 @@ void FSymbolDebugger_AsyncInspect::DoWork()
 //-----------------------------------------------------------------------------
 //	FSymbolDebugger_AsyncSyncFiles
 //-----------------------------------------------------------------------------
-FSymbolDebugger_AsyncSyncFiles::FSymbolDebugger_AsyncSyncFiles(const FString& InSourceControlLabel, const FString& InPlatform)
-	: AskedToAbortCount( 0 )
+FSymbolDebugger_AsyncSyncFiles::FSymbolDebugger_AsyncSyncFiles(const FString& InSourceControlLabel, const FString& InPlatform, ICrashDebugHelper* InCrashHelperModule)
+	: CrashHelperModule(InCrashHelperModule)
+	, AskedToAbortCount( 0 )
 	, SourceControlLabel(InSourceControlLabel)
 	, Platform(InPlatform)
 	, bResult_Succeeded(false)
@@ -74,11 +68,7 @@ FSymbolDebugger_AsyncSyncFiles::FSymbolDebugger_AsyncSyncFiles(const FString& In
 
 void FSymbolDebugger_AsyncSyncFiles::DoWork()
 {
-	FCrashDebugHelperModule& CrashHelperModule = FModuleManager::LoadModuleChecked<FCrashDebugHelperModule>(FName("CrashDebugHelper"));
-	ICrashDebugHelper* Handler = CrashHelperModule.Get();
-	check(Handler != nullptr);
-
-	if (Handler->SyncRequiredFilesForDebuggingFromLabel(SourceControlLabel, Platform) == true)
+	if (CrashHelperModule->SyncRequiredFilesForDebuggingFromLabel(SourceControlLabel, Platform) == true)
 	{
 		bResult_Succeeded = true;
 	}
@@ -107,8 +97,9 @@ void FSymbolDebugger_LaunchDebugger::DoWork()
 //-----------------------------------------------------------------------------
 //	FSymbolDebugger_ProcessCrashDump
 //-----------------------------------------------------------------------------
-FSymbolDebugger_ProcessCrashDump::FSymbolDebugger_ProcessCrashDump(const FString& InCrashDumpName)
-	: AskedToAbortCount( 0 )
+FSymbolDebugger_ProcessCrashDump::FSymbolDebugger_ProcessCrashDump(const FString& InCrashDumpName, ICrashDebugHelper* InCrashHelperModule)
+	: CrashHelperModule(InCrashHelperModule)
+	, AskedToAbortCount( 0 )
 	, CrashDumpName(InCrashDumpName)
 	, bResult_Succeeded(false)
 	, Result_LabelName(TEXT(""))
@@ -122,18 +113,14 @@ void FSymbolDebugger_ProcessCrashDump::DoWork()
 {
 	bResult_Succeeded = false;
 
-	FCrashDebugHelperModule& CrashHelperModule = FModuleManager::LoadModuleChecked<FCrashDebugHelperModule>(FName("CrashDebugHelper"));
-	ICrashDebugHelper* Handler = CrashHelperModule.Get();
-	check(Handler != nullptr);
-
 	FCrashDebugInfo CrashDebugInfo;
-	if (Handler->ParseCrashDump(CrashDumpName, CrashDebugInfo) == true)
+	if (CrashHelperModule->ParseCrashDump(CrashDumpName, CrashDebugInfo) == true)
 	{
 		Result_LabelName = CrashDebugInfo.SourceControlLabel;
 		Result_EngineVersionName = FString::FromInt(CrashDebugInfo.EngineVersion);
 		Result_PlatformName = CrashDebugInfo.PlatformName;
 
-		if (Handler->SyncRequiredFilesForDebuggingFromLabel(Result_LabelName, Result_PlatformName) == true)
+		if (CrashHelperModule->SyncRequiredFilesForDebuggingFromLabel(Result_LabelName, Result_PlatformName) == true)
 		{
 			FPlatformProcess::LaunchFileInDefaultExternalApplication(*CrashDumpName);
 			bResult_Succeeded = true;
@@ -168,10 +155,9 @@ FSymbolDebugger::FSymbolDebugger()
 		SymbolStoreName = LocalSymbolStore;
 	}
 
-	FCrashDebugHelperModule& CrashHelperModule = FModuleManager::LoadModuleChecked<FCrashDebugHelperModule>(FName("CrashDebugHelper"));
-	ICrashDebugHelper* Handler = CrashHelperModule.Get();
-	check(Handler != nullptr);
-	DepotName = Handler->GetDepotName();
+	CrashHelperModule = FModuleManager::LoadModuleChecked<FCrashDebugHelperModule>(FName("CrashDebugHelper")).Get();
+	check(CrashHelperModule != nullptr);
+	DepotName = CrashHelperModule->GetDepotName();
 }
 
 bool FSymbolDebugger::SetCurrentMethod(SSymbolDebugger::ESymbolDebuggerMethods InNewMethod)
@@ -309,10 +295,7 @@ bool FSymbolDebugger::SetTextField(SSymbolDebugger::ESymbolDebuggerTextFields In
 	case SSymbolDebugger::TextField_SourceControlDepot:
 		{
 			DepotName = InNewName;
-			FCrashDebugHelperModule& CrashHelperModule = FModuleManager::LoadModuleChecked<FCrashDebugHelperModule>(FName("CrashDebugHelper"));
-			ICrashDebugHelper* Handler = CrashHelperModule.Get();
-			check(Handler != nullptr);
-			Handler->SetDepotName(DepotName);
+			CrashHelperModule->SetDepotName(DepotName);
 		}
 		break;
 	default:
@@ -535,10 +518,7 @@ void FSymbolDebugger::Tick()
 	if (s_bFirstTick)
 	{
 		// Initialize source control and show the login window
-		FCrashDebugHelperModule& CrashHelperModule = FModuleManager::LoadModuleChecked<FCrashDebugHelperModule>(FName("CrashDebugHelper"));
-		ICrashDebugHelper* Handler = CrashHelperModule.Get();
-		check(Handler != nullptr);
-		Handler->InitSourceControl(true);
+		CrashHelperModule->InitSourceControl(true);
 
 		s_bFirstTick = false;
 	}
@@ -645,7 +625,7 @@ bool FSymbolDebugger::InspectCrashDump(const FString& InCrashDumpName)
 
 	CurrentAction = SSymbolDebugger::DebugAction_Inspect;
 	// Start the async task
-	InspectionTask = MakeShareable(new FAsyncTask<FSymbolDebugger_AsyncInspect>(InCrashDumpName, TEXT(""), TEXT("")));
+	InspectionTask = MakeShareable(new FAsyncTask<FSymbolDebugger_AsyncInspect>(InCrashDumpName, TEXT(""), TEXT(""), CrashHelperModule));
 	InspectionTask->StartBackgroundTask();
 
 	return true;
@@ -663,7 +643,7 @@ bool FSymbolDebugger::InspectEngineVersion(const FString& InEngineVersion)
 
 	CurrentAction = SSymbolDebugger::DebugAction_Inspect;
 	// Start the async task
-	InspectionTask = MakeShareable(new FAsyncTask<FSymbolDebugger_AsyncInspect>(TEXT(""), InEngineVersion, TEXT("")));
+	InspectionTask = MakeShareable(new FAsyncTask<FSymbolDebugger_AsyncInspect>(TEXT(""), InEngineVersion, TEXT(""), CrashHelperModule));
 	InspectionTask->StartBackgroundTask();
 
 	return true;
@@ -681,7 +661,7 @@ bool FSymbolDebugger::InspectChangelist(const FString& InChangelist)
 
 	CurrentAction = SSymbolDebugger::DebugAction_Inspect;
 	// Start the async task
-	InspectionTask = MakeShareable(new FAsyncTask<FSymbolDebugger_AsyncInspect>(TEXT(""), TEXT(""), InChangelist));
+	InspectionTask = MakeShareable(new FAsyncTask<FSymbolDebugger_AsyncInspect>(TEXT(""), TEXT(""), InChangelist, CrashHelperModule));
 	InspectionTask->StartBackgroundTask();
 
 	return true;
@@ -699,7 +679,7 @@ bool FSymbolDebugger::SyncFiles(const FString& InLabelName, const FString& InPla
 
 	CurrentAction = SSymbolDebugger::DebugAction_Sync;
 	// Start the async task
-	SyncTask = MakeShareable(new FAsyncTask<FSymbolDebugger_AsyncSyncFiles>(InLabelName, InPlatform));
+	SyncTask = MakeShareable(new FAsyncTask<FSymbolDebugger_AsyncSyncFiles>(InLabelName, InPlatform, CrashHelperModule));
 	SyncTask->StartBackgroundTask();
 
 	return true;
@@ -735,7 +715,7 @@ bool FSymbolDebugger::ProcessCrashDump(const FString& InCrashDumpName)
 
 	CurrentAction = SSymbolDebugger::DebugAction_Process;
 	// Start the async task
-	ProcessCrashDumpTask = MakeShareable(new FAsyncTask<FSymbolDebugger_ProcessCrashDump>(InCrashDumpName));
+	ProcessCrashDumpTask = MakeShareable(new FAsyncTask<FSymbolDebugger_ProcessCrashDump>(InCrashDumpName, CrashHelperModule));
 	ProcessCrashDumpTask->StartBackgroundTask();
 
 	return true;
diff --git a/Engine/Source/Programs/SymbolDebugger/Private/SymbolDebugger.h b/Engine/Source/Programs/SymbolDebugger/Private/SymbolDebugger.h
--- a/Engine/Source/Programs/SymbolDebugger/Private/SymbolDebugger.h
+++ b/Engine/Source/Programs/SymbolDebugger/Private/SymbolDebugger.h
@@ -12,7 +12,7 @@ class FSymbolDebugger_AsyncInspect : public FNonAbandonableTask
 {
 public:
 	/** Constructor */
-	FSymbolDebugger_AsyncInspect(const FString& InCrashDumpName, const FString& InEngineVersion, const FString& InChangelist);
+	FSymbolDebugger_AsyncInspect(const FString& InCrashDumpName, const FString& InEngineVersion, const FString& InChangelist, ICrashDebugHelper* InCrashHelperModule);
 
 	/** Performs work on thread */
 	void DoWork();
@@ -48,6 +48,9 @@ public:
 	}
 
 private:
+
+	ICrashDebugHelper* CrashHelperModule;
+
 	/** True if we've been asked to abort work in progress at the next opportunity */
 	FThreadSafeCounter AskedToAbortCount;
 
@@ -73,7 +76,7 @@ class FSymbolDebugger_AsyncSyncFiles : public FNonAbandonableTask
 {
 public:
 	/** Constructor */
-	FSymbolDebugger_AsyncSyncFiles(const FString& InSourceControlLabel, const FString& InPlatform);
+	FSymbolDebugger_AsyncSyncFiles(const FString& InSourceControlLabel, const FString& InPlatform, ICrashDebugHelper* InCrashHelperModule);
 
 	/** Performs work on thread */
 	void DoWork();
@@ -99,6 +102,9 @@ public:
 	}
 
 private:
+
+	ICrashDebugHelper* CrashHelperModule;
+
 	/** True if we've been asked to abort work in progress at the next opportunity */
 	FThreadSafeCounter AskedToAbortCount;
 
@@ -161,7 +167,7 @@ class FSymbolDebugger_ProcessCrashDump : public FNonAbandonableTask
 {
 public:
 	/** Constructor */
-	FSymbolDebugger_ProcessCrashDump(const FString& InCrashDumpName);
+	FSymbolDebugger_ProcessCrashDump(const FString& InCrashDumpName, ICrashDebugHelper* InCrashHelperModule);
 
 	/** Performs work on thread */
 	void DoWork();
@@ -205,6 +211,9 @@ public:
 	}
 
 private:
+
+	ICrashDebugHelper* CrashHelperModule;
+
 	/** True if we've been asked to abort work in progress at the next opportunity */
 	FThreadSafeCounter AskedToAbortCount;
 
@@ -448,4 +457,6 @@ protected:
 	TSharedPtr<FAsyncTask<class FSymbolDebugger_ProcessCrashDump> > ProcessCrashDumpTask;
 	/** The results from processing the crash dump */
 	bool bProcessCrashDumpSucceeded;
+
+	ICrashDebugHelper* CrashHelperModule;
 };
