Compiler fix found while doing a batch build

diff --git a/Engine/Source/Programs/SymbolDebugger/Private/SSymbolDebugger.cpp b/Engine/Source/Programs/SymbolDebugger/Private/SSymbolDebugger.cpp
--- a/Engine/Source/Programs/SymbolDebugger/Private/SSymbolDebugger.cpp
+++ b/Engine/Source/Programs/SymbolDebugger/Private/SSymbolDebugger.cpp
@@ -2,6 +2,7 @@
 
 #include "SymbolDebuggerApp.h"
 #include "SSymbolDebugger.h"
+#include "SThrobber.h"
 
 TSharedRef<SWidget> SSymbolDebugger::GenerateActionButton(ESymbolDebuggerActions InAction)
 {
