Stats - Stats for async task running on the thread pool

diff --git a/Engine/Source/Developer/CrashTracker/Private/CrashVideoCapture.cpp b/Engine/Source/Developer/CrashTracker/Private/CrashVideoCapture.cpp
--- a/Engine/Source/Developer/CrashTracker/Private/CrashVideoCapture.cpp
+++ b/Engine/Source/Developer/CrashTracker/Private/CrashVideoCapture.cpp
@@ -71,12 +71,6 @@ void FAsyncImageCompress::DoWork()
 }
 
 
-const TCHAR* FAsyncImageCompress::Name()
-{
-	return TEXT("FAsyncImageCompress");
-}
-
-
 FString FCrashVideoCapture::CaptureVideoPath;
 
 
diff --git a/Engine/Source/Developer/CrashTracker/Private/CrashVideoCapture.h b/Engine/Source/Developer/CrashTracker/Private/CrashVideoCapture.h
--- a/Engine/Source/Developer/CrashTracker/Private/CrashVideoCapture.h
+++ b/Engine/Source/Developer/CrashTracker/Private/CrashVideoCapture.h
@@ -35,7 +35,10 @@ public:
 
 	void DoWork();
 
-	static const TCHAR* Name();
+	FORCEINLINE TStatId GetStatId() const
+	{
+		RETURN_QUICK_DECLARE_CYCLE_STAT(FAsyncImageCompress, STATGROUP_ThreadPoolAsyncTasks);
+	}
 
 private:
 	/** The image compression algorithm to use */
diff --git a/Engine/Source/Developer/DerivedDataCache/Private/DerivedDataBackendAsyncPutWrapper.h b/Engine/Source/Developer/DerivedDataCache/Private/DerivedDataBackendAsyncPutWrapper.h
--- a/Engine/Source/Developer/DerivedDataCache/Private/DerivedDataBackendAsyncPutWrapper.h
+++ b/Engine/Source/Developer/DerivedDataCache/Private/DerivedDataBackendAsyncPutWrapper.h
@@ -84,12 +84,10 @@ public:
 		FilesInFlight->Remove(CacheKey);
 		FDerivedDataBackend::Get().AddToAsyncCompletionCounter(-1);
 	}
-	/** Give the name for external event viewers
-	* @return	the name to display in external event viewers
-	*/
-	static const TCHAR *Name()
+
+	FORCEINLINE TStatId GetStatId() const
 	{
-		return TEXT("FCachePutAsyncWorker");
+		RETURN_QUICK_DECLARE_CYCLE_STAT(FCachePutAsyncWorker, STATGROUP_ThreadPoolAsyncTasks);
 	}
 
 	/** Indicates to the thread pool that this task is abandonable */
diff --git a/Engine/Source/Developer/DerivedDataCache/Private/DerivedDataCache.cpp b/Engine/Source/Developer/DerivedDataCache/Private/DerivedDataCache.cpp
--- a/Engine/Source/Developer/DerivedDataCache/Private/DerivedDataCache.cpp
+++ b/Engine/Source/Developer/DerivedDataCache/Private/DerivedDataCache.cpp
@@ -98,12 +98,10 @@ class FDerivedDataCache : public FDerivedDataCacheInterface
 			}
 			FDerivedDataBackend::Get().AddToAsyncCompletionCounter(-1);
 		}
-		/** Give the name for external event viewers
-		 * @return	the name to display in external event viewers
-		**/
-		static const TCHAR *Name()
+
+		FORCEINLINE TStatId GetStatId() const
 		{
-			return TEXT("FBuildAsyncWorker");
+			RETURN_QUICK_DECLARE_CYCLE_STAT(FBuildAsyncWorker, STATGROUP_ThreadPoolAsyncTasks);
 		}
 
 		/** true in the case of a cache hit, otherwise the result of the deriver build call **/
diff --git a/Engine/Source/Developer/MeshUtilities/Private/MeshUtilities.cpp b/Engine/Source/Developer/MeshUtilities/Private/MeshUtilities.cpp
--- a/Engine/Source/Developer/MeshUtilities/Private/MeshUtilities.cpp
+++ b/Engine/Source/Developer/MeshUtilities/Private/MeshUtilities.cpp
@@ -244,9 +244,9 @@ public:
 
 	void DoWork();
 
-	static const TCHAR* Name()
+	FORCEINLINE TStatId GetStatId() const
 	{
-		return TEXT("MeshDistanceFieldAsyncTask");
+		RETURN_QUICK_DECLARE_CYCLE_STAT(FMeshDistanceFieldAsyncTask, STATGROUP_ThreadPoolAsyncTasks);
 	}
 
 	bool WasNegativeAtBorder() const 
diff --git a/Engine/Source/Developer/ProfilerClient/Private/ProfilerClientManager.h b/Engine/Source/Developer/ProfilerClient/Private/ProfilerClientManager.h
--- a/Engine/Source/Developer/ProfilerClient/Private/ProfilerClientManager.h
+++ b/Engine/Source/Developer/ProfilerClient/Private/ProfilerClientManager.h
@@ -135,9 +135,9 @@ class FProfilerClientManager
 
 		void DoWork();
 
-		static const TCHAR* Name()
+		FORCEINLINE TStatId GetStatId() const
 		{
-			return TEXT("FAsyncReadWorker");
+			RETURN_QUICK_DECLARE_CYCLE_STAT(FAsyncReadWorker, STATGROUP_ThreadPoolAsyncTasks);
 		}
 	};
 
diff --git a/Engine/Source/Developer/TextureCompressor/Private/TextureCompressorModule.cpp b/Engine/Source/Developer/TextureCompressor/Private/TextureCompressorModule.cpp
--- a/Engine/Source/Developer/TextureCompressor/Private/TextureCompressorModule.cpp
+++ b/Engine/Source/Developer/TextureCompressor/Private/TextureCompressorModule.cpp
@@ -1453,13 +1453,9 @@ public:
 			);
 	}
 
-	/** 
-	 * Give the name for external event viewers
-	 * @return	the name to display in external event viewers
-	 */
-	static const TCHAR* Name()
+	FORCEINLINE TStatId GetStatId() const
 	{
-		return TEXT("FAsyncCompressionTask");
+		RETURN_QUICK_DECLARE_CYCLE_STAT(FAsyncCompressionWorker, STATGROUP_ThreadPoolAsyncTasks);
 	}
 
 	bool GetCompressionResults(FCompressedImage2D& OutCompressedImage) const
diff --git a/Engine/Source/Developer/TextureFormatDXT/Private/TextureFormatDXT.cpp b/Engine/Source/Developer/TextureFormatDXT/Private/TextureFormatDXT.cpp
--- a/Engine/Source/Developer/TextureFormatDXT/Private/TextureFormatDXT.cpp
+++ b/Engine/Source/Developer/TextureFormatDXT/Private/TextureFormatDXT.cpp
@@ -244,9 +244,9 @@ public:
 		bCompressionResults = Compressor->Compress();
 	}
 
-	static const TCHAR* Name()
+	FORCEINLINE TStatId GetStatId() const
 	{
-		return TEXT("FAsyncNVTTTask");
+		RETURN_QUICK_DECLARE_CYCLE_STAT(FAsyncNVTTWorker, STATGROUP_ThreadPoolAsyncTasks);
 	}
 
 	/** Retrieve compression results. */
diff --git a/Engine/Source/Developer/TextureFormatIntelISPCTexComp/Private/TextureFormatIntelISPCTexComp.cpp b/Engine/Source/Developer/TextureFormatIntelISPCTexComp/Private/TextureFormatIntelISPCTexComp.cpp
--- a/Engine/Source/Developer/TextureFormatIntelISPCTexComp/Private/TextureFormatIntelISPCTexComp.cpp
+++ b/Engine/Source/Developer/TextureFormatIntelISPCTexComp/Private/TextureFormatIntelISPCTexComp.cpp
@@ -172,9 +172,9 @@ class FTextureFormatIntelISPCTexComp : public ITextureFormat
 						IntelBC6HCompressScans(mpEncSettings, mpInImage, mpOutImage, mYStart, mYEnd);
 					}
 
-					static const TCHAR* Name()
+					FORCEINLINE TStatId GetStatId() const
 					{
-						return TEXT("FIntelCompressWorker");
+						RETURN_QUICK_DECLARE_CYCLE_STAT(FIntelCompressWorker, STATGROUP_ThreadPoolAsyncTasks);
 					}
 
 					bc6h_enc_settings*	mpEncSettings;
@@ -244,9 +244,9 @@ class FTextureFormatIntelISPCTexComp : public ITextureFormat
 						IntelBC7CompressScans(mpEncSettings, mpInImage, mpOutImage, mYStart, mYEnd);
 					}
 
-					static const TCHAR* Name()
+					FORCEINLINE TStatId GetStatId() const
 					{
-						return TEXT("FIntelCompressWorker");
+						RETURN_QUICK_DECLARE_CYCLE_STAT(FIntelCompressWorker, STATGROUP_ThreadPoolAsyncTasks);
 					}
 
 					bc7_enc_settings*	mpEncSettings;
diff --git a/Engine/Source/Editor/Kismet/Private/FindInBlueprints.cpp b/Engine/Source/Editor/Kismet/Private/FindInBlueprints.cpp
--- a/Engine/Source/Editor/Kismet/Private/FindInBlueprints.cpp
+++ b/Engine/Source/Editor/Kismet/Private/FindInBlueprints.cpp
@@ -1467,6 +1467,11 @@ public:
 
 	/** Function to call once the search result is ready */
 	TFunction<void(const FSearchResult&)> OnResultReady;
+
+	FORCEINLINE TStatId GetStatId() const
+	{
+		RETURN_QUICK_DECLARE_CYCLE_STAT( FExtractionTask, STATGROUP_ThreadPoolAsyncTasks );
+	}
 };
 
 TArray<FSearchData> FStreamSearch::MakeExtractionBatch() const
diff --git a/Engine/Source/Editor/UnrealEd/Private/Factories/TTFontImport.cpp b/Engine/Source/Editor/UnrealEd/Private/Factories/TTFontImport.cpp
--- a/Engine/Source/Editor/UnrealEd/Private/Factories/TTFontImport.cpp
+++ b/Engine/Source/Editor/UnrealEd/Private/Factories/TTFontImport.cpp
@@ -370,12 +370,9 @@ class FTextureAlphaToDistanceField
 		/** Called by the thread pool to do the work in this task */
 		void DoWork(void);
 
-		/** Give the name for external event viewers
-		* @return	the name to display in external event viewers
-		*/
-		static const TCHAR *Name()
+		FORCEINLINE TStatId GetStatId() const
 		{
-			return TEXT("FBuildDistanceFieldTask");
+			RETURN_QUICK_DECLARE_CYCLE_STAT(FBuildDistanceFieldTask, STATGROUP_ThreadPoolAsyncTasks);
 		}
 
 	private:
diff --git a/Engine/Source/Editor/UnrealEd/Private/Fbx/FbxSkeletalMeshImport.cpp b/Engine/Source/Editor/UnrealEd/Private/Fbx/FbxSkeletalMeshImport.cpp
--- a/Engine/Source/Editor/UnrealEd/Private/Fbx/FbxSkeletalMeshImport.cpp
+++ b/Engine/Source/Editor/UnrealEd/Private/Fbx/FbxSkeletalMeshImport.cpp
@@ -2660,9 +2660,9 @@ public:
 		MeshUtilities->BuildSkeletalMesh( TempSkeletalMesh->GetImportedResource()->LODModels[0], TempSkeletalMesh->RefSkeleton, LODInfluences, LODWedges, LODFaces, LODPoints, LODPointToRawMap, bKeepOverlappingVertices);
 	}
 
-	static const TCHAR *Name()
+	FORCEINLINE TStatId GetStatId() const
 	{
-		return TEXT("FAsyncImportMorphTargetWork");
+		RETURN_QUICK_DECLARE_CYCLE_STAT(FAsyncImportMorphTargetWork, STATGROUP_ThreadPoolAsyncTasks);
 	}
 
 private:
diff --git a/Engine/Source/Editor/UnrealEd/Private/SourceCodeNavigation.cpp b/Engine/Source/Editor/UnrealEd/Private/SourceCodeNavigation.cpp
--- a/Engine/Source/Editor/UnrealEd/Private/SourceCodeNavigation.cpp
+++ b/Engine/Source/Editor/UnrealEd/Private/SourceCodeNavigation.cpp
@@ -151,13 +151,11 @@ public:
 		return AskedToAbortCount.GetValue() > 0;
 	}
 
-	/** @return Queries the name of this task for for external event viewers */
-	static const TCHAR* Name()
+	FORCEINLINE TStatId GetStatId() const
 	{
-		return TEXT("FAsyncSymbolGatherer");
+		RETURN_QUICK_DECLARE_CYCLE_STAT(FAsyncSymbolGatherer, STATGROUP_ThreadPoolAsyncTasks);
 	}
 
-
 private:
 
 	/** True if we've been asked to abort work in progress at the next opportunity */
@@ -666,10 +664,9 @@ void FSourceCodeNavigation::Initialize()
 			return false;
 		}
 
-		/** @return Queries the name of this task for for external event viewers */
-		static const TCHAR* Name()
+		FORCEINLINE TStatId GetStatId() const
 		{
-			return TEXT("FAsyncInitializeSourceFileDatabase");
+			RETURN_QUICK_DECLARE_CYCLE_STAT(FAsyncInitializeSourceFileDatabase, STATGROUP_ThreadPoolAsyncTasks);
 		}
 	};
 
diff --git a/Engine/Source/Programs/SymbolDebugger/Private/SymbolDebugger.h b/Engine/Source/Programs/SymbolDebugger/Private/SymbolDebugger.h
--- a/Engine/Source/Programs/SymbolDebugger/Private/SymbolDebugger.h
+++ b/Engine/Source/Programs/SymbolDebugger/Private/SymbolDebugger.h
@@ -23,10 +23,9 @@ public:
 		return AskedToAbortCount.GetValue() > 0;
 	}
 
-	/** @return Queries the name of this task for for external event viewers */
-	static const TCHAR* Name()
+	FORCEINLINE TStatId GetStatId() const
 	{
-		return TEXT("FSymbolDebugger_AsyncInspect");
+		RETURN_QUICK_DECLARE_CYCLE_STAT( FSymbolDebugger_AsyncInspect, STATGROUP_ThreadPoolAsyncTasks );
 	}
 
 	/** The found name */
@@ -87,10 +86,9 @@ public:
 		return AskedToAbortCount.GetValue() > 0;
 	}
 
-	/** @return Queries the name of this task for for external event viewers */
-	static const TCHAR* Name()
+	FORCEINLINE TStatId GetStatId() const
 	{
-		return TEXT("FSymbolDebugger_AsyncSyncFiles");
+		RETURN_QUICK_DECLARE_CYCLE_STAT( FSymbolDebugger_AsyncSyncFiles, STATGROUP_ThreadPoolAsyncTasks );
 	}
 
 	/** 
@@ -135,10 +133,9 @@ public:
 		return AskedToAbortCount.GetValue() > 0;
 	}
 
-	/** @return Queries the name of this task for for external event viewers */
-	static const TCHAR* Name()
+	FORCEINLINE TStatId GetStatId() const
 	{
-		return TEXT("FSymbolDebugger_LaunchDebugger");
+		RETURN_QUICK_DECLARE_CYCLE_STAT( FSymbolDebugger_LaunchDebugger, STATGROUP_ThreadPoolAsyncTasks );
 	}
 
 	/** 
@@ -178,10 +175,9 @@ public:
 		return AskedToAbortCount.GetValue() > 0;
 	}
 
-	/** @return Queries the name of this task for for external event viewers */
-	static const TCHAR* Name()
+	FORCEINLINE TStatId GetStatId() const
 	{
-		return TEXT("FSymbolDebugger_ProcessCrashDump");
+		RETURN_QUICK_DECLARE_CYCLE_STAT( FSymbolDebugger_ProcessCrashDump, STATGROUP_ThreadPoolAsyncTasks );
 	}
 
 	/** The found name */
diff --git a/Engine/Source/Runtime/Core/Private/Misc/RemoteConfigIni.cpp b/Engine/Source/Runtime/Core/Private/Misc/RemoteConfigIni.cpp
--- a/Engine/Source/Runtime/Core/Private/Misc/RemoteConfigIni.cpp
+++ b/Engine/Source/Runtime/Core/Private/Misc/RemoteConfigIni.cpp
@@ -84,11 +84,6 @@ FRemoteConfigAsyncIOInfo& FRemoteConfigAsyncWorker::GetIOInfo()
 }
 
 
-const TCHAR* FRemoteConfigAsyncWorker::Name()
-{
-	return TEXT("FRemoteConfigAsyncWorker");
-}
-
 
 bool FRemoteConfigAsyncWorker::CanAbandon()
 {
diff --git a/Engine/Source/Runtime/Core/Private/Serialization/Archive.cpp b/Engine/Source/Runtime/Core/Private/Serialization/Archive.cpp
--- a/Engine/Source/Runtime/Core/Private/Serialization/Archive.cpp
+++ b/Engine/Source/Runtime/Core/Private/Serialization/Archive.cpp
@@ -303,12 +303,9 @@ public:
 		verify( FCompression::CompressMemory( Flags, CompressedBuffer, CompressedSize, UncompressedBuffer, UncompressedSize ) );
 	}
 
-	/** Give the name for external event viewers
-	* @return	the name to display in external event viewers
-	*/
-	static const TCHAR *Name()
+	FORCEINLINE TStatId GetStatId() const
 	{
-		return TEXT("FAsyncCompressionChunks");
+		RETURN_QUICK_DECLARE_CYCLE_STAT(FAsyncCompressionChunk, STATGROUP_ThreadPoolAsyncTasks);
 	}
 };
 #endif		// WITH_MULTI_THREADED_COMPRESSION
diff --git a/Engine/Source/Runtime/Core/Public/Misc/RemoteConfigIni.h b/Engine/Source/Runtime/Core/Public/Misc/RemoteConfigIni.h
--- a/Engine/Source/Runtime/Core/Public/Misc/RemoteConfigIni.h
+++ b/Engine/Source/Runtime/Core/Public/Misc/RemoteConfigIni.h
@@ -57,8 +57,10 @@ public:
 	/** Returns the local IO info object */
 	FRemoteConfigAsyncIOInfo& GetIOInfo();
 
-	/** Give the name for external event viewers. */
-	static const TCHAR* Name();
+	FORCEINLINE TStatId GetStatId() const
+	{
+		RETURN_QUICK_DECLARE_CYCLE_STAT(FRemoteConfigAsyncWorker, STATGROUP_ThreadPoolAsyncTasks);
+	}
 
 	/** Indicates to the thread pool that this task is abandonable. */
 	bool CanAbandon();
diff --git a/Engine/Source/Runtime/Core/Public/Misc/SecureHash.h b/Engine/Source/Runtime/Core/Public/Misc/SecureHash.h
--- a/Engine/Source/Runtime/Core/Public/Misc/SecureHash.h
+++ b/Engine/Source/Runtime/Core/Public/Misc/SecureHash.h
@@ -289,12 +289,10 @@ public:
 			Buffer = 0;
 		}
 	}
-	/** Give the name for external event viewers
-	* @return	the name to display in external event viewers
-	*/
-	static const TCHAR *Name()
+
+	FORCEINLINE TStatId GetStatId() const
 	{
-		return TEXT("FAsyncSHAVerify");
+		RETURN_QUICK_DECLARE_CYCLE_STAT(FAsyncSHAVerify, STATGROUP_ThreadPoolAsyncTasks);
 	}
 };
 
diff --git a/Engine/Source/Runtime/CoreUObject/Private/UObject/SavePackage.cpp b/Engine/Source/Runtime/CoreUObject/Private/UObject/SavePackage.cpp
--- a/Engine/Source/Runtime/CoreUObject/Private/UObject/SavePackage.cpp
+++ b/Engine/Source/Runtime/CoreUObject/Private/UObject/SavePackage.cpp
@@ -191,12 +191,10 @@ void AsyncWriteFile(const TArray<uint8>& Data, const TCHAR* Filename, const FDat
 			}
 			OutstandingAsyncWrites.Decrement();
 		}
-		/** Give the name for external event viewers
-		* @return	the name to display in external event viewers
-		*/
-		static const TCHAR *Name()
+
+		FORCEINLINE TStatId GetStatId() const
 		{
-			return TEXT("FAsyncWriteWorker");
+			RETURN_QUICK_DECLARE_CYCLE_STAT(FAsyncWriteWorker, STATGROUP_ThreadPoolAsyncTasks);
 		}
 	};
 
diff --git a/Engine/Source/Runtime/Engine/Private/AudioDerivedData.cpp b/Engine/Source/Runtime/Engine/Private/AudioDerivedData.cpp
--- a/Engine/Source/Runtime/Engine/Private/AudioDerivedData.cpp
+++ b/Engine/Source/Runtime/Engine/Private/AudioDerivedData.cpp
@@ -359,10 +359,9 @@ public:
 		}
 	}
 
-	/** Interface for FAsyncTask. */
-	static const TCHAR* Name()
+	FORCEINLINE TStatId GetStatId() const
 	{
-		return TEXT("FStreamedAudioAsyncCacheDerivedDataTask");
+		RETURN_QUICK_DECLARE_CYCLE_STAT(FStreamedAudioCacheDerivedDataWorker, STATGROUP_ThreadPoolAsyncTasks);
 	}
 };
 
diff --git a/Engine/Source/Runtime/Engine/Private/ContentStreaming.cpp b/Engine/Source/Runtime/Engine/Private/ContentStreaming.cpp
--- a/Engine/Source/Runtime/Engine/Private/ContentStreaming.cpp
+++ b/Engine/Source/Runtime/Engine/Private/ContentStreaming.cpp
@@ -1202,13 +1202,9 @@ private:
 		PrioritizedTextures.Sort( FCompareTexturePriority() );
 	}
 
-	/**
-	 * Give the name for external event viewers
-	 * @return	the name to display in external event viewers
-	 */ 
-	static const TCHAR *Name()
+	FORCEINLINE TStatId GetStatId() const
 	{
-		return TEXT("FAsyncTextureStreaming");
+		RETURN_QUICK_DECLARE_CYCLE_STAT(FAsyncTextureStreaming, STATGROUP_ThreadPoolAsyncTasks);
 	}
 
 	/** Reference to the owning streaming manager, for accessing the thread-safe data. */
diff --git a/Engine/Source/Runtime/Engine/Private/TextureDerivedData.cpp b/Engine/Source/Runtime/Engine/Private/TextureDerivedData.cpp
--- a/Engine/Source/Runtime/Engine/Private/TextureDerivedData.cpp
+++ b/Engine/Source/Runtime/Engine/Private/TextureDerivedData.cpp
@@ -912,10 +912,9 @@ public:
 		}
 	}
 
-	/** Interface for FAsyncTask. */
-	static const TCHAR* Name()
+	FORCEINLINE TStatId GetStatId() const
 	{
-		return TEXT("FTextureAsyncCacheDerivedDataTask");
+		RETURN_QUICK_DECLARE_CYCLE_STAT(FTextureCacheDerivedDataWorker, STATGROUP_ThreadPoolAsyncTasks);
 	}
 };
 
diff --git a/Engine/Source/Runtime/Engine/Public/AI/Navigation/RecastNavMeshGenerator.h b/Engine/Source/Runtime/Engine/Public/AI/Navigation/RecastNavMeshGenerator.h
--- a/Engine/Source/Runtime/Engine/Public/AI/Navigation/RecastNavMeshGenerator.h
+++ b/Engine/Source/Runtime/Engine/Public/AI/Navigation/RecastNavMeshGenerator.h
@@ -240,9 +240,9 @@ struct ENGINE_API FRecastTileGeneratorWrapper : public FNonAbandonableTask
 		TileGenerator->DoWork();
 	}
 
-	static const TCHAR *Name()
+	FORCEINLINE TStatId GetStatId() const
 	{
-		return TEXT("FRecastTileGenerator");
+		RETURN_QUICK_DECLARE_CYCLE_STAT(FRecastTileGenerator, STATGROUP_ThreadPoolAsyncTasks);
 	}
 };
 
diff --git a/Engine/Source/Runtime/Engine/Public/AudioDecompress.h b/Engine/Source/Runtime/Engine/Public/AudioDecompress.h
--- a/Engine/Source/Runtime/Engine/Public/AudioDecompress.h
+++ b/Engine/Source/Runtime/Engine/Public/AudioDecompress.h
@@ -138,12 +138,9 @@ public:
 	 */
 	ENGINE_API void DoWork();
 
-	/** Give the name for external event viewers
-	* @return	the name to display in external event viewers
-	*/
-	static const TCHAR *Name()
+	FORCEINLINE TStatId GetStatId() const
 	{
-		return TEXT("FAsyncAudioDecompress");
+		RETURN_QUICK_DECLARE_CYCLE_STAT(FAsyncAudioDecompressWorker, STATGROUP_ThreadPoolAsyncTasks);
 	}
 };
 
diff --git a/Engine/Source/Runtime/Engine/Public/AudioStreaming.h b/Engine/Source/Runtime/Engine/Public/AudioStreaming.h
--- a/Engine/Source/Runtime/Engine/Public/AudioStreaming.h
+++ b/Engine/Source/Runtime/Engine/Public/AudioStreaming.h
@@ -39,12 +39,9 @@ public:
 	 */
 	void DoWork();
 
-	/**
-	 * Name of the async task.
-	 */
-	static const TCHAR* Name()
+	FORCEINLINE TStatId GetStatId() const
 	{
-		return TEXT("FAsyncStreamDerivedChunkTask");
+		RETURN_QUICK_DECLARE_CYCLE_STAT(FAsyncStreamDerivedChunkWorker, STATGROUP_ThreadPoolAsyncTasks);
 	}
 
 	/**
diff --git a/Engine/Source/Runtime/Engine/Public/TextureResource.h b/Engine/Source/Runtime/Engine/Public/TextureResource.h
--- a/Engine/Source/Runtime/Engine/Public/TextureResource.h
+++ b/Engine/Source/Runtime/Engine/Public/TextureResource.h
@@ -61,12 +61,9 @@ public:
 	 */
 	void DoWork();
 
-	/**
-	 * Name of the async task.
-	 */
-	static const TCHAR* Name()
+	FORCEINLINE TStatId GetStatId() const
 	{
-		return TEXT("FAsyncStreamDerivedMipTask");
+		RETURN_QUICK_DECLARE_CYCLE_STAT(FAsyncStreamDerivedMipWorker, STATGROUP_ThreadPoolAsyncTasks);
 	}
 
 	/**
@@ -130,14 +127,14 @@ public:
 		check(Args.ThreadSafeCounter);
 	}
 
-	static const TCHAR* Name()
-	{
-		return TEXT("FCreateTextureTask");
-	}
-
 	/** Creates the texture. */
 	void DoWork();
 
+	FORCEINLINE TStatId GetStatId() const
+	{
+		RETURN_QUICK_DECLARE_CYCLE_STAT(FCreateTextureTask, STATGROUP_ThreadPoolAsyncTasks);
+	}
+
 private:
 	/** Task arguments. */
 	FArguments Args;
diff --git a/Engine/Source/Runtime/Landscape/Classes/LandscapeProxy.h b/Engine/Source/Runtime/Landscape/Classes/LandscapeProxy.h
--- a/Engine/Source/Runtime/Landscape/Classes/LandscapeProxy.h
+++ b/Engine/Source/Runtime/Landscape/Classes/LandscapeProxy.h
@@ -295,9 +295,10 @@ public:
 	{
 	}
 	void DoWork();
-	static const TCHAR *Name()
+
+	FORCEINLINE TStatId GetStatId() const
 	{
-		return TEXT("FAsyncGrassTask");
+		RETURN_QUICK_DECLARE_CYCLE_STAT(FAsyncGrassTask, STATGROUP_ThreadPoolAsyncTasks);
 	}
 
 	~FAsyncGrassTask();
diff --git a/Engine/Source/Runtime/NetworkFile/Private/NetworkPlatformFile.cpp b/Engine/Source/Runtime/NetworkFile/Private/NetworkPlatformFile.cpp
--- a/Engine/Source/Runtime/NetworkFile/Private/NetworkPlatformFile.cpp
+++ b/Engine/Source/Runtime/NetworkFile/Private/NetworkPlatformFile.cpp
@@ -822,12 +822,10 @@ public:
 			}
 		}
 	}
-	/** Give the name for external event viewers
-	* @return	the name to display in external event viewers
-	*/
-	static const TCHAR *Name()
+
+	FORCEINLINE TStatId GetStatId() const
 	{
-		return TEXT("FAsyncNetworkWriteWorker");
+		RETURN_QUICK_DECLARE_CYCLE_STAT(FAsyncNetworkWriteWorker, STATGROUP_ThreadPoolAsyncTasks);
 	}
 };
 
@@ -897,12 +895,10 @@ void AsyncReadUnsolicitedFiles(int32 InNumUnsolictedFiles, FNetworkPlatformFile&
 			}
 			NetworkDoneEvent->Trigger();
 		}
-		/** Give the name for external event viewers
-		* @return	the name to display in external event viewers
-		*/
-		static const TCHAR *Name()
+
+		FORCEINLINE TStatId GetStatId() const
 		{
-			return TEXT("FAsyncReadUnsolicitedFile");
+			RETURN_QUICK_DECLARE_CYCLE_STAT(FAsyncReadUnsolicitedFile, STATGROUP_ThreadPoolAsyncTasks);
 		}
 	};
 
diff --git a/Engine/Source/Runtime/PakFile/Private/IPlatformFilePak.cpp b/Engine/Source/Runtime/PakFile/Private/IPlatformFilePak.cpp
--- a/Engine/Source/Runtime/PakFile/Private/IPlatformFilePak.cpp
+++ b/Engine/Source/Runtime/PakFile/Private/IPlatformFilePak.cpp
@@ -101,9 +101,10 @@ public:
 				FMemory::Memcpy(CopyOut, UncompressedBuffer+CopyOffset, CopyLength);
 			}
 		}
-		static const TCHAR *Name()
+
+		FORCEINLINE TStatId GetStatId() const
 		{
-			return TEXT("FPakUncompressTask");
+			RETURN_QUICK_DECLARE_CYCLE_STAT(FPakUncompressTask, STATGROUP_ThreadPoolAsyncTasks);
 		}
 	};
 
diff --git a/Engine/Source/Runtime/Sockets/Public/IPAddress.h b/Engine/Source/Runtime/Sockets/Public/IPAddress.h
--- a/Engine/Source/Runtime/Sockets/Public/IPAddress.h
+++ b/Engine/Source/Runtime/Sockets/Public/IPAddress.h
@@ -245,12 +245,10 @@ class FResolveInfoAsync :
 		{
 			Parent->DoWork();
 		}
-		/** Give the name for external event viewers
-		* @return	the name to display in external event viewers
-		*/
-		static const TCHAR *Name()
+
+		FORCEINLINE TStatId GetStatId() const
 		{
-			return TEXT("FResolveInfoAsyncWorker");
+			RETURN_QUICK_DECLARE_CYCLE_STAT(FResolveInfoAsyncWorker, STATGROUP_ThreadPoolAsyncTasks);
 		}
 
 		/** Indicates to the thread pool that this task is abandonable */
