SymbolDebugger - Fixed a crash during startup (bCompileAgainstCoreUObject=true)

diff --git a/Engine/Source/Programs/SymbolDebugger/SymbolDebugger.Target.cs b/Engine/Source/Programs/SymbolDebugger/SymbolDebugger.Target.cs
--- a/Engine/Source/Programs/SymbolDebugger/SymbolDebugger.Target.cs
+++ b/Engine/Source/Programs/SymbolDebugger/SymbolDebugger.Target.cs
@@ -53,7 +53,7 @@ public class SymbolDebuggerTarget : TargetRules
 
 		// SymbolDebugger doesn't ever compile with the engine linked in
 		UEBuildConfiguration.bCompileAgainstEngine = false;
-		UEBuildConfiguration.bCompileAgainstCoreUObject = false;
+		UEBuildConfiguration.bCompileAgainstCoreUObject = true;
 
 		UEBuildConfiguration.bIncludeADO = true;
 
