CrashReport - Commented obsolete functionality

diff --git a/Engine/Source/Programs/SymbolDebugger/Private/SymbolDebugger.cpp b/Engine/Source/Programs/SymbolDebugger/Private/SymbolDebugger.cpp
--- a/Engine/Source/Programs/SymbolDebugger/Private/SymbolDebugger.cpp
+++ b/Engine/Source/Programs/SymbolDebugger/Private/SymbolDebugger.cpp
@@ -68,14 +68,14 @@ FSymbolDebugger_AsyncSyncFiles::FSymbolDebugger_AsyncSyncFiles(const FString& In
 
 void FSymbolDebugger_AsyncSyncFiles::DoWork()
 {
-	if (CrashHelperModule->SyncRequiredFilesForDebuggingFromLabel(SourceControlLabel, Platform) == true)
-	{
-		bResult_Succeeded = true;
-	}
-	else
-	{
-		bResult_Succeeded = false;
-	}
+// 	if (CrashHelperModule->SyncRequiredFilesForDebuggingFromLabel(SourceControlLabel, Platform) == true)
+// 	{
+// 		bResult_Succeeded = true;
+// 	}
+// 	else
+// 	{
+// 		bResult_Succeeded = false;
+// 	}
 }
 
 //-----------------------------------------------------------------------------
