CrashReport - Commented obsolete functionality

diff --git a/Engine/Source/Programs/SymbolDebugger/Private/SymbolDebugger.cpp b/Engine/Source/Programs/SymbolDebugger/Private/SymbolDebugger.cpp
--- a/Engine/Source/Programs/SymbolDebugger/Private/SymbolDebugger.cpp
+++ b/Engine/Source/Programs/SymbolDebugger/Private/SymbolDebugger.cpp
@@ -120,11 +120,11 @@ void FSymbolDebugger_ProcessCrashDump::DoWork()
 		Result_EngineVersionName = FString::FromInt(CrashDebugInfo.EngineVersion);
 		Result_PlatformName = CrashDebugInfo.PlatformName;
 
-		if (CrashHelperModule->SyncRequiredFilesForDebuggingFromLabel(Result_LabelName, Result_PlatformName) == true)
-		{
-			FPlatformProcess::LaunchFileInDefaultExternalApplication(*CrashDumpName);
-			bResult_Succeeded = true;
-		}
+// 		if (CrashHelperModule->SyncRequiredFilesForDebuggingFromLabel(Result_LabelName, Result_PlatformName) == true)
+// 		{
+// 			FPlatformProcess::LaunchFileInDefaultExternalApplication(*CrashDumpName);
+// 			bResult_Succeeded = true;
+// 		}
 	}
 }
 
