CrashReport - Commented obsolete functionality

diff --git a/Engine/Source/Programs/SymbolDebugger/Private/SymbolDebugger.cpp b/Engine/Source/Programs/SymbolDebugger/Private/SymbolDebugger.cpp
--- a/Engine/Source/Programs/SymbolDebugger/Private/SymbolDebugger.cpp
+++ b/Engine/Source/Programs/SymbolDebugger/Private/SymbolDebugger.cpp
@@ -21,36 +21,36 @@ FSymbolDebugger_AsyncInspect::FSymbolDebugger_AsyncInspect(const FString& InCras
 
 void FSymbolDebugger_AsyncInspect::DoWork()
 {
-	if (CrashDumpName.Len() > 0)
-	{
-		FCrashDebugInfo CrashDebugInfo;
-		if (CrashHelperModule->ParseCrashDump(CrashDumpName, CrashDebugInfo) == true)
-		{
-			Result_LabelName = CrashDebugInfo.SourceControlLabel;
-			Result_EngineVersionName = FString::FromInt(CrashDebugInfo.EngineVersion);
-			Result_PlatformName = CrashDebugInfo.PlatformName;
-		}
-		else
-		{
-			Result_LabelName = TEXT("");
-			Result_EngineVersionName = TEXT("");
-			Result_PlatformName = TEXT("");
-		}
-	}
-	else if ((EngineVersionName.Len() > 0) || (ChangelistName.Len() > 0))
-	{
-		const FString FoundLabel;// = CrashHelperModule->GetLabelFromChangelistNumber( FCString::Atoi( *ChangelistName ) );
-
-		if (FoundLabel.Len() > 0)
-		{
-			Result_EngineVersionName = TEXT( "" );
-			Result_LabelName = FoundLabel;
-		}
-		else
-		{
-			Result_LabelName = TEXT("");
-		}
-	}
+// 	if (CrashDumpName.Len() > 0)
+// 	{
+// 		FCrashDebugInfo CrashDebugInfo;
+// 		if (CrashHelperModule->ParseCrashDump(CrashDumpName, CrashDebugInfo) == true)
+// 		{
+// 			Result_LabelName = CrashDebugInfo.SourceControlLabel;
+// 			Result_EngineVersionName = FString::FromInt(CrashDebugInfo.EngineVersion);
+// 			Result_PlatformName = CrashDebugInfo.PlatformName;
+// 		}
+// 		else
+// 		{
+// 			Result_LabelName = TEXT("");
+// 			Result_EngineVersionName = TEXT("");
+// 			Result_PlatformName = TEXT("");
+// 		}
+// 	}
+// 	else if ((EngineVersionName.Len() > 0) || (ChangelistName.Len() > 0))
+// 	{
+// 		const FString FoundLabel;// = CrashHelperModule->GetLabelFromChangelistNumber( FCString::Atoi( *ChangelistName ) );
+// 
+// 		if (FoundLabel.Len() > 0)
+// 		{
+// 			Result_EngineVersionName = TEXT( "" );
+// 			Result_LabelName = FoundLabel;
+// 		}
+// 		else
+// 		{
+// 			Result_LabelName = TEXT("");
+// 		}
+// 	}
 }
 
 
@@ -111,21 +111,21 @@ FSymbolDebugger_ProcessCrashDump::FSymbolDebugger_ProcessCrashDump(const FString
 /** Performs work on thread */
 void FSymbolDebugger_ProcessCrashDump::DoWork()
 {
-	bResult_Succeeded = false;
-
-	FCrashDebugInfo CrashDebugInfo;
-	if (CrashHelperModule->ParseCrashDump(CrashDumpName, CrashDebugInfo) == true)
-	{
-		Result_LabelName = CrashDebugInfo.SourceControlLabel;
-		Result_EngineVersionName = FString::FromInt(CrashDebugInfo.EngineVersion);
-		Result_PlatformName = CrashDebugInfo.PlatformName;
-
-// 		if (CrashHelperModule->SyncRequiredFilesForDebuggingFromLabel(Result_LabelName, Result_PlatformName) == true)
-// 		{
-// 			FPlatformProcess::LaunchFileInDefaultExternalApplication(*CrashDumpName);
-// 			bResult_Succeeded = true;
-// 		}
-	}
+// 	bResult_Succeeded = false;
+// 
+// 	FCrashDebugInfo CrashDebugInfo;
+// 	if (CrashHelperModule->ParseCrashDump(CrashDumpName, CrashDebugInfo) == true)
+// 	{
+// 		Result_LabelName = CrashDebugInfo.SourceControlLabel;
+// 		Result_EngineVersionName = FString::FromInt(CrashDebugInfo.EngineVersion);
+// 		Result_PlatformName = CrashDebugInfo.PlatformName;
+// 
+//  		if (CrashHelperModule->SyncRequiredFilesForDebuggingFromLabel(Result_LabelName, Result_PlatformName) == true)
+//  		{
+//  			FPlatformProcess::LaunchFileInDefaultExternalApplication(*CrashDumpName);
+//  			bResult_Succeeded = true;
+//  		}
+// 	}
 }
 
 //-----------------------------------------------------------------------------
